Dataflow Editor:
- Added a missing cloth menu content update before building the visualization menu.
- Removed unnecessary override of SimulationSceneUpdated that was misleadingly refreshing menu content instead of doing any sort of meaningful visualization uptate.

[RN] Dataflow Editor - Added a missing cloth menu content update before building the visualization menu.

#Synced-CL 50399480

diff --git a/Engine/Plugins/ChaosClothAssetEditorCore/Source/ChaosClothAssetEditor/Private/ChaosClothAsset/ClothDataflowSimulationVisualization.cpp b/Engine/Plugins/ChaosClothAssetEditorCore/Source/ChaosClothAssetEditor/Private/ChaosClothAsset/ClothDataflowSimulationVisualization.cpp
--- a/Engine/Plugins/ChaosClothAssetEditorCore/Source/ChaosClothAssetEditor/Private/ChaosClothAsset/ClothDataflowSimulationVisualization.cpp
+++ b/Engine/Plugins/ChaosClothAssetEditorCore/Source/ChaosClothAssetEditor/Private/ChaosClothAsset/ClothDataflowSimulationVisualization.cpp
@@ -25,8 +25,9 @@ namespace UE::Chaos::ClothAsset
 			{
 				if (const TSharedPtr<FDataflowSimulationScene>& SimulationScene = Toolkit->GetSimulationScene())
 				{
-					if (GetClothComponent(SimulationScene.Get()))
+					if (const UChaosClothComponent* const ClothComponent = GetClothComponent(SimulationScene.Get()))
 					{
+						ClothEditorSimulationVisualization.RefreshMenusForClothComponent(ClothComponent);  // Update the menus' content before building them
 						ClothEditorSimulationVisualization.ExtendViewportShowMenu(MenuBuilder, ViewportClient.ToSharedRef());
 					}
 				}
@@ -61,11 +62,6 @@ namespace UE::Chaos::ClothAsset
 		return ClothEditorSimulationVisualization.GetDisplayString(GetClothComponent(SimulationScene));
 	}
 
-	void FClothDataflowSimulationVisualization::SimulationSceneUpdated(const FDataflowSimulationScene* SimulationScene)
-	{
-		ClothEditorSimulationVisualization.RefreshMenusForClothComponent(GetClothComponent(SimulationScene));
-	}
-
 }  // End namespace UE::Chaos::ClothAsset
 
 #undef LOCTEXT_NAMESPACE 
diff --git a/Engine/Plugins/ChaosClothAssetEditorCore/Source/ChaosClothAssetEditor/Private/ChaosClothAsset/ClothDataflowSimulationVisualization.h b/Engine/Plugins/ChaosClothAssetEditorCore/Source/ChaosClothAssetEditor/Private/ChaosClothAsset/ClothDataflowSimulationVisualization.h
--- a/Engine/Plugins/ChaosClothAssetEditorCore/Source/ChaosClothAssetEditor/Private/ChaosClothAsset/ClothDataflowSimulationVisualization.h
+++ b/Engine/Plugins/ChaosClothAssetEditorCore/Source/ChaosClothAssetEditor/Private/ChaosClothAsset/ClothDataflowSimulationVisualization.h
@@ -21,7 +21,6 @@ namespace UE::Chaos::ClothAsset
 		virtual void Draw(const FDataflowSimulationScene* SimulationScene, FPrimitiveDrawInterface* PDI) override;
 		virtual void DrawCanvas(const FDataflowSimulationScene* SimulationScene, FCanvas* Canvas, const FSceneView* SceneView) override;
 		virtual FText GetDisplayString(const FDataflowSimulationScene* SimulationScene) const override;
-		virtual void SimulationSceneUpdated(const FDataflowSimulationScene* SimulationScene) override;
 
 		const UChaosClothComponent* GetClothComponent(const FDataflowSimulationScene* SimulationScene) const;
 
