Chaos Cloth: Add options to filter debug draw to a single vertex or named set of vertices.

diff --git a/Engine/Plugins/ChaosClothAssetEditorCore/Source/ChaosClothAssetEditor/Private/ChaosClothAsset/ClothEditorSimulationVisualization.cpp b/Engine/Plugins/ChaosClothAssetEditorCore/Source/ChaosClothAssetEditor/Private/ChaosClothAsset/ClothEditorSimulationVisualization.cpp
--- a/Engine/Plugins/ChaosClothAssetEditorCore/Source/ChaosClothAssetEditor/Private/ChaosClothAsset/ClothEditorSimulationVisualization.cpp
+++ b/Engine/Plugins/ChaosClothAssetEditorCore/Source/ChaosClothAssetEditor/Private/ChaosClothAsset/ClothEditorSimulationVisualization.cpp
@@ -93,9 +93,9 @@ static FText GetSimulationStatisticsString(const FClothSimulationProxy& SimProxy
 	return TextValue;
 }
 
-DECLARE_DELEGATE_ThreeParams(FClothVisualizationDebugDraw, const FClothEditorSimulationVisualization&, const ::Chaos::FClothVisualizationNoGC&, FPrimitiveDrawInterface*);
+DECLARE_DELEGATE_ThreeParams(FClothVisualizationDebugDraw, const FClothEditorSimulationVisualization&, const ::Chaos::FClothVisualizationNoGC&, const ::Chaos::FClothVisualizationNoGC::FDrawContext&);
 
-DECLARE_DELEGATE_FourParams(FClothVisualizationDebugDrawTexts, const FClothEditorSimulationVisualization&, const ::Chaos::FClothVisualizationNoGC&, FCanvas*, const FSceneView*);
+DECLARE_DELEGATE_ThreeParams(FClothVisualizationDebugDrawTexts, const FClothEditorSimulationVisualization&, const ::Chaos::FClothVisualizationNoGC&, const ::Chaos::FClothVisualizationNoGC::FDrawTextsContext&);
 
 DECLARE_DELEGATE_RetVal_TwoParams(FText, FLocalDebugDisplayString, const FClothEditorSimulationVisualization&, const FClothSimulationProxy& );
 
@@ -106,6 +106,8 @@ struct FVisualizationOption
 	// Actual option entries
 	static const FVisualizationOption OptionData[];
 	static const uint32 Count;
+	static const uint32 SingleVertexOptionIndex = 1;
+	static const uint32 FilterSetOptionIndex = 2;
 
 	// Chaos debug draw function
 	FClothVisualizationDebugDraw ClothVisualizationDebugDraw;
@@ -149,8 +151,9 @@ struct FVisualizationOption
 		, bDefaultFlagValue(bInDefaultFlagValue)
 	{}
 	
-	FVisualizationOption(const FText & InDisplayName, const FText & InToolTip, bool bInDisablesSimulation = false, bool bInHidesClothSections = false, bool bInDefaultFlagValue = false)
-		: DisplayName(InDisplayName)
+	FVisualizationOption(const FText & InDisplayName, const FText & InToolTip, bool bInDisablesSimulation = false, bool bInHidesClothSections = false, bool bInDefaultFlagValue = false, const FAdditionalMenus& InAdditionalMenus = FAdditionalMenus())
+		: AdditionalMenus(InAdditionalMenus)
+		, DisplayName(InDisplayName)
 		, ToolTip(InToolTip)
 		, bDisablesSimulation(bInDisablesSimulation)
 		, bHidesClothSections(bInHidesClothSections)
@@ -173,24 +176,43 @@ const FVisualizationOption FVisualizationOption::OptionData[] =
 		LOCTEXT("ChaosVisName_SimulationStatistics_Tooltip", "Displays simulation statistics"),
 		/*bDisablesSimulation =*/false, /*bHidesClothSections=*/false, /*bDefaultFlagValue=*/true),
 	FVisualizationOption(
-		FClothVisualizationDebugDraw::CreateLambda([](const FClothEditorSimulationVisualization&, const ::Chaos::FClothVisualizationNoGC& Visualization, FPrimitiveDrawInterface* PDI)
+		LOCTEXT("ChaosVisName_SingleVertex", "Draw Single Vertex"),
+		LOCTEXT("ChaosVisName_SingleVertex_Tooltip", "Draw will be filtered to a single vertex"),
+		/*bDisablesSimulation =*/false, /*bHidesClothSections=*/false, /*bDefaultFlagValue=*/false,
+		FAdditionalMenus::CreateLambda([](FClothEditorSimulationVisualization& EditorVisualization, FMenuBuilder& MenuBuilder)
+		{
+			EditorVisualization.ExtendViewportShowMenuSingleVertexSelection(MenuBuilder);
+		})),
+	FVisualizationOption(
+		LOCTEXT("ChaosVisName_FilterSet", "Filter Set"),
+		LOCTEXT("ChaosVisName_FilterSet_Tooltip", "Draw will be filtered to a SimVertex set"),
+		/*bDisablesSimulation =*/false, /*bHidesClothSections=*/false, /*bDefaultFlagValue=*/false,
+		FAdditionalMenus::CreateLambda([](FClothEditorSimulationVisualization& EditorVisualization, FMenuBuilder& MenuBuilder)
+		{
+			EditorVisualization.ExtendViewportShowMenuFilterSetSelection(MenuBuilder);
+		})),
+
+// DO NOT ADD ABOVE THIS LINE WITHOUT UPDATING FVisualizationOption::SingleVertexOptionIndex and FVisualizationOption::FilterSetOptionIndex
+
+	FVisualizationOption(
+		FClothVisualizationDebugDraw::CreateLambda([](const FClothEditorSimulationVisualization&, const ::Chaos::FClothVisualizationNoGC& Visualization, const ::Chaos::FClothVisualizationNoGC::FDrawContext& Context)
 		{
-			Visualization.DrawPhysMeshShaded(PDI);
+			Visualization.DrawPhysMeshShaded(Context);
 		}),
 		LOCTEXT("ChaosVisName_PhysMesh", "Physical Mesh (Flat Shaded)"),
 		LOCTEXT("ChaosVisName_PhysMeshShaded_ToolTip", "Draws the current physical result as a doubled sided flat shaded mesh"),
 		/*bDisablesSimulation =*/false, /*bHidesClothSections=*/true),
 	FVisualizationOption(
-		FClothVisualizationDebugDraw::CreateLambda([](const FClothEditorSimulationVisualization&, const ::Chaos::FClothVisualizationNoGC& Visualization, FPrimitiveDrawInterface* PDI)
+		FClothVisualizationDebugDraw::CreateLambda([](const FClothEditorSimulationVisualization&, const ::Chaos::FClothVisualizationNoGC& Visualization, const ::Chaos::FClothVisualizationNoGC::FDrawContext& Context)
 		{
-			Visualization.DrawPhysMeshWired(PDI);
+			Visualization.DrawPhysMeshWired(Context);
 		}), 
 		LOCTEXT("ChaosVisName_PhysMeshWire", "Physical Mesh (Wireframe)"), 
 		LOCTEXT("ChaosVisName_PhysMeshWired_ToolTip", "Draws the current physical mesh result in wireframe")),
 	FVisualizationOption(
-		FClothVisualizationDebugDraw::CreateLambda([](const FClothEditorSimulationVisualization&, const ::Chaos::FClothVisualizationNoGC& Visualization, FPrimitiveDrawInterface* PDI)
+		FClothVisualizationDebugDraw::CreateLambda([](const FClothEditorSimulationVisualization&, const ::Chaos::FClothVisualizationNoGC& Visualization, const ::Chaos::FClothVisualizationNoGC::FDrawContext& Context)
 		{
-			Visualization.DrawAnimMeshWired(PDI);
+			Visualization.DrawAnimMeshWired(Context);
 		}),
 		LOCTEXT("ChaosVisName_AnimMeshWire", "Animated Mesh (Wireframe)"), 
 		LOCTEXT("ChaosVisName_AnimMeshWired_ToolTip", "Draws the current animated mesh input in wireframe")),
@@ -199,23 +221,23 @@ const FVisualizationOption FVisualizationOption::OptionData[] =
 		LOCTEXT("ChaosVisName_HideRenderMesh_ToolTip", "Hide the render mesh."), 
 		/*bDisablesSimulation =*/false, /*bHidesClothSections=*/true),
 	FVisualizationOption(
-		FClothVisualizationDebugDrawTexts::CreateLambda([](const FClothEditorSimulationVisualization&, const ::Chaos::FClothVisualizationNoGC& Visualization, FCanvas* Canvas, const FSceneView* SceneView)
+		FClothVisualizationDebugDrawTexts::CreateLambda([](const FClothEditorSimulationVisualization&, const ::Chaos::FClothVisualizationNoGC& Visualization, const ::Chaos::FClothVisualizationNoGC::FDrawTextsContext& Context)
 		{
-			Visualization.DrawParticleIndices(Canvas, SceneView);
+			Visualization.DrawParticleIndices(Context);
 		}),
 		LOCTEXT("ChaosVisName_ParticleIndices", "Particle Indices"), 
 		LOCTEXT("ChaosVisName_ParticleIndices_ToolTip", "Draws the particle indices as instantiated by the solver")),
 	FVisualizationOption(
-		FClothVisualizationDebugDrawTexts::CreateLambda([](const FClothEditorSimulationVisualization&, const ::Chaos::FClothVisualizationNoGC& Visualization, FCanvas* Canvas, const FSceneView* SceneView)
+		FClothVisualizationDebugDrawTexts::CreateLambda([](const FClothEditorSimulationVisualization&, const ::Chaos::FClothVisualizationNoGC& Visualization, const ::Chaos::FClothVisualizationNoGC::FDrawTextsContext& Context)
 		{
-			Visualization.DrawElementIndices(Canvas, SceneView);
+			Visualization.DrawElementIndices(Context);
 		}),
 		LOCTEXT("ChaosVisName_ElementIndices", "Element Indices"), 
 		LOCTEXT("ChaosVisName_ElementIndices_ToolTip", "Draws the element's (triangle or other) indices as instantiated by the solver")),
 	FVisualizationOption(
-		FClothVisualizationDebugDraw::CreateLambda([](const FClothEditorSimulationVisualization& EditorVisualization, const ::Chaos::FClothVisualizationNoGC& Visualization, FPrimitiveDrawInterface* PDI)
+		FClothVisualizationDebugDraw::CreateLambda([](const FClothEditorSimulationVisualization& EditorVisualization, const ::Chaos::FClothVisualizationNoGC& Visualization, const ::Chaos::FClothVisualizationNoGC::FDrawContext& Context)
 		{
-			Visualization.DrawPointNormals(PDI, EditorVisualization.GetPointNormalLength());
+			Visualization.DrawPointNormals(Context, EditorVisualization.GetPointNormalLength());
 		}),
 		LOCTEXT("ChaosVisName_PointNormals", "Physical Mesh Normals"), 
 		LOCTEXT("ChaosVisName_PointNormals_ToolTip", "Draws the current point normals for the simulation mesh"),
@@ -225,16 +247,16 @@ const FVisualizationOption FVisualizationOption::OptionData[] =
 			EditorVisualization.ExtendViewportShowMenuPointNormalsLength(MenuBuilder);
 		})),
 	FVisualizationOption(
-		FClothVisualizationDebugDraw::CreateLambda([](const FClothEditorSimulationVisualization&, const ::Chaos::FClothVisualizationNoGC& Visualization, FPrimitiveDrawInterface* PDI)
+		FClothVisualizationDebugDraw::CreateLambda([](const FClothEditorSimulationVisualization&, const ::Chaos::FClothVisualizationNoGC& Visualization, const ::Chaos::FClothVisualizationNoGC::FDrawContext& Context)
 		{
-			Visualization.DrawPointVelocities(PDI);
+			Visualization.DrawPointVelocities(Context);
 		}),
 		LOCTEXT("ChaosVisName_PointVelocities", "Point Velocities"), 
 		LOCTEXT("ChaosVisName_PointVelocities_ToolTip", "Draws the current point velocities for the simulation mesh")),
 	FVisualizationOption(
-		FClothVisualizationDebugDraw::CreateLambda([](const FClothEditorSimulationVisualization& EditorVisualization, const ::Chaos::FClothVisualizationNoGC& Visualization, FPrimitiveDrawInterface* PDI)
+		FClothVisualizationDebugDraw::CreateLambda([](const FClothEditorSimulationVisualization& EditorVisualization, const ::Chaos::FClothVisualizationNoGC& Visualization, const ::Chaos::FClothVisualizationNoGC::FDrawContext& Context)
 		{
-			Visualization.DrawAnimNormals(PDI, EditorVisualization.GetAnimatedNormalLength());
+			Visualization.DrawAnimNormals(Context, EditorVisualization.GetAnimatedNormalLength());
 		}),
 		LOCTEXT("ChaosVisName_AnimNormals", "Animated Mesh Normals"), 
 		LOCTEXT("ChaosVisName_AnimNormals_ToolTip", "Draws the current point normals for the animated mesh"),
@@ -244,82 +266,82 @@ const FVisualizationOption FVisualizationOption::OptionData[] =
 			EditorVisualization.ExtendViewportShowMenuAnimatedNormalsLength(MenuBuilder);
 		})),
 	FVisualizationOption(
-		FClothVisualizationDebugDraw::CreateLambda([](const FClothEditorSimulationVisualization&, const ::Chaos::FClothVisualizationNoGC& Visualization, FPrimitiveDrawInterface* PDI)
+		FClothVisualizationDebugDraw::CreateLambda([](const FClothEditorSimulationVisualization&, const ::Chaos::FClothVisualizationNoGC& Visualization, const ::Chaos::FClothVisualizationNoGC::FDrawContext& Context)
 		{
 			constexpr bool bForceWireframe = false;
-			Visualization.DrawCollision(PDI, bForceWireframe);
+			Visualization.DrawCollision(Context, bForceWireframe);
 		}),
 		LOCTEXT("ChaosVisName_Collision", "Collisions"), 
 		LOCTEXT("ChaosVisName_Collision_ToolTip", "Draws the collision bodies the simulation is currently using")),
 	FVisualizationOption(
-		FClothVisualizationDebugDraw::CreateLambda([](const FClothEditorSimulationVisualization&, const ::Chaos::FClothVisualizationNoGC& Visualization, FPrimitiveDrawInterface* PDI)
+		FClothVisualizationDebugDraw::CreateLambda([](const FClothEditorSimulationVisualization&, const ::Chaos::FClothVisualizationNoGC& Visualization, const ::Chaos::FClothVisualizationNoGC::FDrawContext& Context)
 		{
 			constexpr bool bForceWireframe = true;
-			Visualization.DrawCollision(PDI, bForceWireframe);
+			Visualization.DrawCollision(Context, bForceWireframe);
 		}),
 		LOCTEXT("ChaosVisName_CollisionWireframe", "Collisions (Force Wireframe)"),
 		LOCTEXT("ChaosVisName_Collision_ToolTip", "Draws the collision bodies the simulation is currently using")),
 	FVisualizationOption(
-		FClothVisualizationDebugDraw::CreateLambda([](const FClothEditorSimulationVisualization&, const ::Chaos::FClothVisualizationNoGC& Visualization, FPrimitiveDrawInterface* PDI)
+		FClothVisualizationDebugDraw::CreateLambda([](const FClothEditorSimulationVisualization&, const ::Chaos::FClothVisualizationNoGC& Visualization, const ::Chaos::FClothVisualizationNoGC::FDrawContext& Context)
 		{
-			Visualization.DrawBackstops(PDI);
+			Visualization.DrawBackstops(Context);
 		}),
 		LOCTEXT("ChaosVisName_Backstop", "Backstops"), 
 		LOCTEXT("ChaosVisName_Backstop_ToolTip", "Draws the backstop radius and position for each simulation particle")),
 	FVisualizationOption(
-		FClothVisualizationDebugDraw::CreateLambda([](const FClothEditorSimulationVisualization&, const ::Chaos::FClothVisualizationNoGC& Visualization, FPrimitiveDrawInterface* PDI)
+		FClothVisualizationDebugDraw::CreateLambda([](const FClothEditorSimulationVisualization&, const ::Chaos::FClothVisualizationNoGC& Visualization, const ::Chaos::FClothVisualizationNoGC::FDrawContext& Context)
 		{
-			Visualization.DrawBackstopDistances(PDI);
+			Visualization.DrawBackstopDistances(Context);
 		}),
 		LOCTEXT("ChaosVisName_BackstopDistance", "Backstop Distances"), 
 		LOCTEXT("ChaosVisName_BackstopDistance_ToolTip", "Draws the backstop distance offset for each simulation particle"), 
 		/*bDisablesSimulation =*/true),
 	FVisualizationOption(
-		FClothVisualizationDebugDraw::CreateLambda([](const FClothEditorSimulationVisualization&, const ::Chaos::FClothVisualizationNoGC& Visualization, FPrimitiveDrawInterface* PDI)
+		FClothVisualizationDebugDraw::CreateLambda([](const FClothEditorSimulationVisualization&, const ::Chaos::FClothVisualizationNoGC& Visualization, const ::Chaos::FClothVisualizationNoGC::FDrawContext& Context)
 		{
-			Visualization.DrawMaxDistances(PDI);
+			Visualization.DrawMaxDistances(Context);
 		}),
 		LOCTEXT("ChaosVisName_MaxDistance", "Max Distances"), 
 		LOCTEXT("ChaosVisName_MaxDistance_ToolTip", "Draws the current max distances for the sim particles as a line along its normal")),
 	FVisualizationOption(
-		FClothVisualizationDebugDrawTexts::CreateLambda([](const FClothEditorSimulationVisualization&, const ::Chaos::FClothVisualizationNoGC& Visualization, FCanvas* Canvas, const FSceneView* SceneView)
+		FClothVisualizationDebugDrawTexts::CreateLambda([](const FClothEditorSimulationVisualization&, const ::Chaos::FClothVisualizationNoGC& Visualization, const ::Chaos::FClothVisualizationNoGC::FDrawTextsContext& Context)
 		{
-			Visualization.DrawMaxDistanceValues(Canvas, SceneView);
+			Visualization.DrawMaxDistanceValues(Context);
 		}),
 		LOCTEXT("ChaosVisName_MaxDistanceValue", "Max Distances As Numbers"), 
 		LOCTEXT("ChaosVisName_MaxDistanceValue_ToolTip", "Draws the current max distances as numbers")),
 	FVisualizationOption(
-		FClothVisualizationDebugDraw::CreateLambda([](const FClothEditorSimulationVisualization&, const ::Chaos::FClothVisualizationNoGC& Visualization, FPrimitiveDrawInterface* PDI)
+		FClothVisualizationDebugDraw::CreateLambda([](const FClothEditorSimulationVisualization&, const ::Chaos::FClothVisualizationNoGC& Visualization, const ::Chaos::FClothVisualizationNoGC::FDrawContext& Context)
 		{
-			Visualization.DrawAnimDrive(PDI);
+			Visualization.DrawAnimDrive(Context);
 		}),
 		LOCTEXT("ChaosVisName_AnimDrive", "Anim Drive"), 
 		LOCTEXT("ChaosVisName_AnimDrive_Tooltip", "Draws the current skinned reference mesh for the simulation which anim drive will attempt to reach if enabled")),
 	FVisualizationOption(
-		FClothVisualizationDebugDraw::CreateLambda([](const FClothEditorSimulationVisualization&, const ::Chaos::FClothVisualizationNoGC& Visualization, FPrimitiveDrawInterface* PDI)
+		FClothVisualizationDebugDraw::CreateLambda([](const FClothEditorSimulationVisualization&, const ::Chaos::FClothVisualizationNoGC& Visualization, const ::Chaos::FClothVisualizationNoGC::FDrawContext& Context)
 		{
-			Visualization.DrawEdgeConstraint(PDI);
+			Visualization.DrawEdgeConstraint(Context);
 		}),
 		LOCTEXT("ChaosVisName_EdgeConstraint", "Edge Constraint"),
 		LOCTEXT("ChaosVisName_EdgeConstraint_Tooltip", "Draws the edge spring constraints")),
 	FVisualizationOption(
-		FClothVisualizationDebugDraw::CreateLambda([](const FClothEditorSimulationVisualization&, const ::Chaos::FClothVisualizationNoGC& Visualization, FPrimitiveDrawInterface* PDI)
+		FClothVisualizationDebugDraw::CreateLambda([](const FClothEditorSimulationVisualization&, const ::Chaos::FClothVisualizationNoGC& Visualization, const ::Chaos::FClothVisualizationNoGC::FDrawContext& Context)
 		{
-			Visualization.DrawBendingConstraint(PDI);
+			Visualization.DrawBendingConstraint(Context);
 		}),
 		LOCTEXT("ChaosVisName_BendingConstraint", "Bending Constraint"), 
 		LOCTEXT("ChaosVisName_BendingConstraint_Tooltip", "Draws the bending spring constraints")),
 	FVisualizationOption(
-		FClothVisualizationDebugDraw::CreateLambda([](const FClothEditorSimulationVisualization&, const ::Chaos::FClothVisualizationNoGC& Visualization, FPrimitiveDrawInterface* PDI)
+		FClothVisualizationDebugDraw::CreateLambda([](const FClothEditorSimulationVisualization&, const ::Chaos::FClothVisualizationNoGC& Visualization, const ::Chaos::FClothVisualizationNoGC::FDrawContext& Context)
 		{
-			Visualization.DrawLongRangeConstraint(PDI);
+			Visualization.DrawLongRangeConstraint(Context);
 		}),
 		LOCTEXT("ChaosVisName_LongRangeConstraint" , "Long Range Constraint"), 
 		LOCTEXT("ChaosVisName_LongRangeConstraint_Tooltip", "Draws the long range attachment constraint distances")),
 	FVisualizationOption(
-		FClothVisualizationDebugDraw::CreateLambda([](const FClothEditorSimulationVisualization& EditorVisualization, const ::Chaos::FClothVisualizationNoGC& Visualization, FPrimitiveDrawInterface* PDI)
+		FClothVisualizationDebugDraw::CreateLambda([](const FClothEditorSimulationVisualization& EditorVisualization, const ::Chaos::FClothVisualizationNoGC& Visualization, const ::Chaos::FClothVisualizationNoGC::FDrawContext& Context)
 		{
-			Visualization.DrawWindAndPressureForces(PDI, EditorVisualization.GetAerodynamicsLengthScale());
+			Visualization.DrawWindAndPressureForces(Context, EditorVisualization.GetAerodynamicsLengthScale());
 		}),
 		LOCTEXT("ChaosVisName_WindAndPressureForces", "Wind Aerodynamic And Pressure Forces"), 
 		LOCTEXT("ChaosVisName_WindAndPressure_Tooltip", "Draws the Wind drag and lift and pressure forces"),
@@ -329,9 +351,9 @@ const FVisualizationOption FVisualizationOption::OptionData[] =
 			EditorVisualization.ExtendViewportShowMenuAerodynamicsLengthScale(MenuBuilder);
 		})),
 	FVisualizationOption(
-		FClothVisualizationDebugDraw::CreateLambda([](const FClothEditorSimulationVisualization& EditorVisualization, const ::Chaos::FClothVisualizationNoGC& Visualization, FPrimitiveDrawInterface* PDI)
+		FClothVisualizationDebugDraw::CreateLambda([](const FClothEditorSimulationVisualization& EditorVisualization, const ::Chaos::FClothVisualizationNoGC& Visualization, const ::Chaos::FClothVisualizationNoGC::FDrawContext& Context)
 		{
-			Visualization.DrawWindVelocity(PDI, EditorVisualization.GetWindVelocityLengthScale());
+			Visualization.DrawWindVelocity(Context, EditorVisualization.GetWindVelocityLengthScale());
 		}),
 		LOCTEXT("ChaosVisName_WindVelocity", "Wind Velocity"), 
 		LOCTEXT("ChaosVisName_WindVelocity_Tooltip", "Draws the Wind Velocity vector"),
@@ -341,9 +363,9 @@ const FVisualizationOption FVisualizationOption::OptionData[] =
 			EditorVisualization.ExtendViewportShowMenuWindVelocityLength(MenuBuilder);
 		})),
 	FVisualizationOption(
-		FClothVisualizationDebugDraw::CreateLambda([](const FClothEditorSimulationVisualization& EditorVisualization, const ::Chaos::FClothVisualizationNoGC& Visualization, FPrimitiveDrawInterface* PDI)
+		FClothVisualizationDebugDraw::CreateLambda([](const FClothEditorSimulationVisualization& EditorVisualization, const ::Chaos::FClothVisualizationNoGC& Visualization, const ::Chaos::FClothVisualizationNoGC::FDrawContext& Context)
 		{
-			Visualization.DrawLocalSpace(PDI, EditorVisualization.GetLocalSpaceLengthScale());
+			Visualization.DrawLocalSpace(Context, EditorVisualization.GetLocalSpaceLengthScale());
 		}),
 		LOCTEXT("ChaosVisName_LocalSpace", "Local Space Reference Bone"), 
 		LOCTEXT("ChaosVisName_LocalSpace_Tooltip", "Draws the local space reference bone"),
@@ -353,16 +375,16 @@ const FVisualizationOption FVisualizationOption::OptionData[] =
 			EditorVisualization.ExtendViewportShowMenuLocalSpaceLengthScale(MenuBuilder);
 		})),
 	FVisualizationOption(
-		FClothVisualizationDebugDrawTexts::CreateLambda([](const FClothEditorSimulationVisualization&, const ::Chaos::FClothVisualizationNoGC& Visualization, FCanvas* Canvas, const FSceneView* SceneView)
+		FClothVisualizationDebugDrawTexts::CreateLambda([](const FClothEditorSimulationVisualization&, const ::Chaos::FClothVisualizationNoGC& Visualization, const ::Chaos::FClothVisualizationNoGC::FDrawTextsContext& Context)
 		{
-			Visualization.DrawLocalSpaceBoneNames(Canvas, SceneView);
+			Visualization.DrawLocalSpaceBoneNames(Context);
 		}),
 		LOCTEXT("ChaosVisName_LocalSpaceBoneNames", "Local Space Reference Bone Names"), 
 		LOCTEXT("ChaosVisName_LocalSpaceBoneNames_ToolTip", "Draws the local space reference bone names")),
 	FVisualizationOption(
-		FClothVisualizationDebugDraw::CreateLambda([](const FClothEditorSimulationVisualization& EditorVisualization, const ::Chaos::FClothVisualizationNoGC& Visualization, FPrimitiveDrawInterface* PDI)
+		FClothVisualizationDebugDraw::CreateLambda([](const FClothEditorSimulationVisualization& EditorVisualization, const ::Chaos::FClothVisualizationNoGC& Visualization, const ::Chaos::FClothVisualizationNoGC::FDrawContext& Context)
 		{
-			Visualization.DrawVelocityScale(PDI, EditorVisualization.GetVelocityScaleLength());
+			Visualization.DrawVelocityScale(Context, EditorVisualization.GetVelocityScaleLength());
 		}),
 		LOCTEXT("ChaosVisName_VelocityScale", "Velocity Scale"), 
 		LOCTEXT("ChaosVisName_VelocityScale_Tooltip", "Draws the velocity scale"),
@@ -372,88 +394,88 @@ const FVisualizationOption FVisualizationOption::OptionData[] =
 			EditorVisualization.ExtendViewportShowMenuVelocityScaleLength(MenuBuilder);
 		})),
 	FVisualizationOption(
-		FClothVisualizationDebugDraw::CreateLambda([](const FClothEditorSimulationVisualization&, const ::Chaos::FClothVisualizationNoGC& Visualization, FPrimitiveDrawInterface* PDI)
+		FClothVisualizationDebugDraw::CreateLambda([](const FClothEditorSimulationVisualization&, const ::Chaos::FClothVisualizationNoGC& Visualization, const ::Chaos::FClothVisualizationNoGC::FDrawContext& Context)
 		{
-			Visualization.DrawGravity(PDI);
+			Visualization.DrawGravity(Context);
 		}),
 		LOCTEXT("ChaosVisName_Gravity", "Gravity"), 
 		LOCTEXT("ChaosVisName_Gravity_Tooltip", "Draws gravity")),
 	FVisualizationOption(
-		FClothVisualizationDebugDraw::CreateLambda([](const FClothEditorSimulationVisualization&, const ::Chaos::FClothVisualizationNoGC& Visualization, FPrimitiveDrawInterface* PDI)
+		FClothVisualizationDebugDraw::CreateLambda([](const FClothEditorSimulationVisualization&, const ::Chaos::FClothVisualizationNoGC& Visualization, const ::Chaos::FClothVisualizationNoGC::FDrawContext& Context)
 		{
-			Visualization.DrawFictitiousAngularForces(PDI);
+			Visualization.DrawFictitiousAngularForces(Context);
 		}),
 		LOCTEXT("ChaosVisName_FictitiousAngularForces", "Fictitious Angular Forces"),
 		LOCTEXT("ChaosVisName_Gravity_FictitiousAngularForces", "Draws fictitious angular forces (force based solver only)")),
 	FVisualizationOption(
-		FClothVisualizationDebugDraw::CreateLambda([](const FClothEditorSimulationVisualization&, const ::Chaos::FClothVisualizationNoGC& Visualization, FPrimitiveDrawInterface* PDI)
+		FClothVisualizationDebugDraw::CreateLambda([](const FClothEditorSimulationVisualization&, const ::Chaos::FClothVisualizationNoGC& Visualization, const ::Chaos::FClothVisualizationNoGC::FDrawContext& Context)
 		{
-			Visualization.DrawSelfCollision(PDI);
+			Visualization.DrawSelfCollision(Context);
 		}),
 		LOCTEXT("ChaosVisName_SelfCollision", "Self Collision"), 
 		LOCTEXT("ChaosVisName_SelfCollision_Tooltip", "Draws the self collision thickness/debugging information")),
 	FVisualizationOption(
-		FClothVisualizationDebugDraw::CreateLambda([](const FClothEditorSimulationVisualization&, const ::Chaos::FClothVisualizationNoGC& Visualization, FPrimitiveDrawInterface* PDI)
+		FClothVisualizationDebugDraw::CreateLambda([](const FClothEditorSimulationVisualization&, const ::Chaos::FClothVisualizationNoGC& Visualization, const ::Chaos::FClothVisualizationNoGC::FDrawContext& Context)
 		{
-			Visualization.DrawSelfIntersection(PDI);
+			Visualization.DrawSelfIntersection(Context);
 		}),
 		LOCTEXT("ChaosVisName_SelfIntersection", "Self Intersection"), 
 		LOCTEXT("ChaosVisName_SelfIntersection_Tooltip", "Draws the self intersection contour/region information")),
 	FVisualizationOption(
-		FClothVisualizationDebugDraw::CreateLambda([](const FClothEditorSimulationVisualization&, const ::Chaos::FClothVisualizationNoGC& Visualization, FPrimitiveDrawInterface* PDI)
+		FClothVisualizationDebugDraw::CreateLambda([](const FClothEditorSimulationVisualization&, const ::Chaos::FClothVisualizationNoGC& Visualization, const ::Chaos::FClothVisualizationNoGC::FDrawContext& Context)
 		{
-			Visualization.DrawSelfCollisionLayers(PDI);
+			Visualization.DrawSelfCollisionLayers(Context);
 		}),
 		LOCTEXT("ChaosVisName_SelfCollisionLayers", "Self Collision Layers"), 
 		LOCTEXT("ChaosVisName_SelfCollisionLayers_Tooltip", "Draws the self collision layers"),
 		/*bDisablesSimulation =*/false, /*bHidesClothSections=*/true),
 	FVisualizationOption(
-		FClothVisualizationDebugDraw::CreateLambda([](const FClothEditorSimulationVisualization&, const ::Chaos::FClothVisualizationNoGC& Visualization, FPrimitiveDrawInterface* PDI)
+		FClothVisualizationDebugDraw::CreateLambda([](const FClothEditorSimulationVisualization&, const ::Chaos::FClothVisualizationNoGC& Visualization, const ::Chaos::FClothVisualizationNoGC::FDrawContext& Context)
 		{
-			Visualization.DrawSelfCollisionThickness(PDI);
+			Visualization.DrawSelfCollisionThickness(Context);
 		}),
 		LOCTEXT("ChaosVisName_SelfCollisionThickness", "Self Collision Thickness"), 
 		LOCTEXT("ChaosVisName_SelfCollisionThickness_Tooltip", "Draws the self collision Thickness")),
 	FVisualizationOption(
-		FClothVisualizationDebugDraw::CreateLambda([](const FClothEditorSimulationVisualization&, const ::Chaos::FClothVisualizationNoGC& Visualization, FPrimitiveDrawInterface* PDI)
+		FClothVisualizationDebugDraw::CreateLambda([](const FClothEditorSimulationVisualization&, const ::Chaos::FClothVisualizationNoGC& Visualization, const ::Chaos::FClothVisualizationNoGC::FDrawContext& Context)
 		{
-			Visualization.DrawKinematicColliderShaded(PDI);
+			Visualization.DrawKinematicColliderShaded(Context);
 		}),
 		LOCTEXT("ChaosVisName_DrawKinematicColliderShaded", "Draw Kinematic Colliders (Shaded)"), 
 		LOCTEXT("ChaosVisName_DrawKinematicColliderShaded_Tooltip", "Draw kinematic cloth colliders with flat shading.")),
 	FVisualizationOption(
-		FClothVisualizationDebugDraw::CreateLambda([](const FClothEditorSimulationVisualization&, const ::Chaos::FClothVisualizationNoGC& Visualization, FPrimitiveDrawInterface* PDI)
+		FClothVisualizationDebugDraw::CreateLambda([](const FClothEditorSimulationVisualization&, const ::Chaos::FClothVisualizationNoGC& Visualization, const ::Chaos::FClothVisualizationNoGC::FDrawContext& Context)
 		{
-			Visualization.DrawKinematicColliderWired(PDI);
+			Visualization.DrawKinematicColliderWired(Context);
 		}),
 		LOCTEXT("ChaosVisName_DrawKinematicColliderWired", "Draw Kinematic Colliders (Wireframe)"), 
 		LOCTEXT("ChaosVisName_DrawKinematicColliderWired_Tooltip", "Draw kinematic cloth colliders in wireframe.")),
 	FVisualizationOption(
-		FClothVisualizationDebugDraw::CreateLambda([](const FClothEditorSimulationVisualization&, const ::Chaos::FClothVisualizationNoGC& Visualization, FPrimitiveDrawInterface* PDI)
+		FClothVisualizationDebugDraw::CreateLambda([](const FClothEditorSimulationVisualization&, const ::Chaos::FClothVisualizationNoGC& Visualization, const ::Chaos::FClothVisualizationNoGC::FDrawContext& Context)
 		{
-			Visualization.DrawClothClothConstraints(PDI);
+			Visualization.DrawClothClothConstraints(Context);
 		}),
 		LOCTEXT("ChaosVisName_DrawClothClothConstraints", "Cloth-Cloth constraints"), 
 		LOCTEXT("ChaosVisName_DrawClothClothConstraints_Tooltip", "Draw cloth-cloth constraints.")),
 	FVisualizationOption(
-		FClothVisualizationDebugDraw::CreateLambda([](const FClothEditorSimulationVisualization&, const ::Chaos::FClothVisualizationNoGC& Visualization, FPrimitiveDrawInterface* PDI)
+		FClothVisualizationDebugDraw::CreateLambda([](const FClothEditorSimulationVisualization&, const ::Chaos::FClothVisualizationNoGC& Visualization, const ::Chaos::FClothVisualizationNoGC::FDrawContext& Context)
 		{
-			Visualization.DrawTeleportReset(PDI);
+			Visualization.DrawTeleportReset(Context);
 		}),
 		LOCTEXT("ChaosVisName_DrawTeleportReset", "Teleport/Reset"), 
 		LOCTEXT("ChaosVisName_DrawTeleportReset_Tooltip", "Draw teleport/reset status.")),
 	FVisualizationOption(
-		FClothVisualizationDebugDraw::CreateLambda([](const FClothEditorSimulationVisualization&, const ::Chaos::FClothVisualizationNoGC& Visualization, FPrimitiveDrawInterface* PDI)
+		FClothVisualizationDebugDraw::CreateLambda([](const FClothEditorSimulationVisualization&, const ::Chaos::FClothVisualizationNoGC& Visualization, const ::Chaos::FClothVisualizationNoGC::FDrawContext& Context)
 		{
-			Visualization.DrawMultiResConstraint(PDI);
+			Visualization.DrawMultiResConstraint(Context);
 		}),
 		LOCTEXT("ChaosVisName_DrawMultiResConstraint", "Draw Multi Res Constraint"), 
 		LOCTEXT("ChaosVisName_DrawMultiResConstraint_Tooltip", "Draw multi res constraint coarse mesh and targets.")),
 	FVisualizationOption(
-		FClothVisualizationDebugDraw::CreateLambda([](const FClothEditorSimulationVisualization& EditorVisualization, const ::Chaos::FClothVisualizationNoGC& Visualization, FPrimitiveDrawInterface* PDI)
+		FClothVisualizationDebugDraw::CreateLambda([](const FClothEditorSimulationVisualization& EditorVisualization, const ::Chaos::FClothVisualizationNoGC& Visualization, const ::Chaos::FClothVisualizationNoGC::FDrawContext& Context)
 		{
 			const FString* const WeightMap = EditorVisualization.GetCurrentlySelectedWeightMap();
-			Visualization.DrawWeightMapWithName(PDI, WeightMap ? *WeightMap : FString());
+			Visualization.DrawWeightMapWithName(Context, WeightMap ? *WeightMap : FString());
 		}),
 		LOCTEXT("ChaosVisName_DrawWeightMap", "Weight Map"), 
 		LOCTEXT("ChaosVisName_DrawWeightMap_ToolTip", "Draw the weight map for the simulation mesh. You can control the name of the map to be visualized by setting the p.ChaosClothVisualization.WeightMapName console variable."), 
@@ -463,10 +485,10 @@ const FVisualizationOption FVisualizationOption::OptionData[] =
 			EditorVisualization.ExtendViewportShowMenuWeightMapSelector(MenuBuilder);
 		})),
 	FVisualizationOption(
-		FClothVisualizationDebugDraw::CreateLambda([](const FClothEditorSimulationVisualization& EditorVisualization, const ::Chaos::FClothVisualizationNoGC& Visualization, FPrimitiveDrawInterface* PDI)
+		FClothVisualizationDebugDraw::CreateLambda([](const FClothEditorSimulationVisualization& EditorVisualization, const ::Chaos::FClothVisualizationNoGC& Visualization, const ::Chaos::FClothVisualizationNoGC::FDrawContext& Context)
 		{
 			const FString* const MorphTargetName = EditorVisualization.GetCurrentlySelectedMorphTarget();
-			Visualization.DrawSimMorphTarget(PDI, MorphTargetName ? *MorphTargetName : FString());
+			Visualization.DrawSimMorphTarget(Context, MorphTargetName ? *MorphTargetName : FString());
 		}),
 		LOCTEXT("ChaosVisName_DrawSimMorphTarget", "Sim Morph Target"), 
 		LOCTEXT("ChaosVisName_DrawSimMorphTarget_ToolTip", "Draw a sim morph target. If none selected, the currently active morph target is displayed."), 
@@ -476,10 +498,10 @@ const FVisualizationOption FVisualizationOption::OptionData[] =
 			EditorVisualization.ExtendViewportShowMenuMorphTargetSelector(MenuBuilder);
 		})),
 	FVisualizationOption(
-		FClothVisualizationDebugDraw::CreateLambda([](const FClothEditorSimulationVisualization& EditorVisualization, const ::Chaos::FClothVisualizationNoGC& Visualization, FPrimitiveDrawInterface* PDI)
+		FClothVisualizationDebugDraw::CreateLambda([](const FClothEditorSimulationVisualization& EditorVisualization, const ::Chaos::FClothVisualizationNoGC& Visualization, const ::Chaos::FClothVisualizationNoGC::FDrawContext& Context)
 		{
 			const FString* const AccessoryMeshName = EditorVisualization.GetCurrentlySelectedAccessoryMesh();
-			Visualization.DrawAccessoryMesh(PDI, AccessoryMeshName ? FName(*AccessoryMeshName) : NAME_None);
+			Visualization.DrawAccessoryMesh(Context, AccessoryMeshName ? FName(*AccessoryMeshName) : NAME_None);
 		}),
 		LOCTEXT("ChaosVisName_DrawAccessoryMesh", "Accessory Mesh"), 
 		LOCTEXT("ChaosVisName_DrawAccessoryMesh_ToolTip", "Draw an accessory mesh."), 
@@ -489,10 +511,10 @@ const FVisualizationOption FVisualizationOption::OptionData[] =
 			EditorVisualization.ExtendViewportShowMenuAccessoryMeshSelector(MenuBuilder);
 		})),
 	FVisualizationOption(
-		FClothVisualizationDebugDraw::CreateLambda([](const FClothEditorSimulationVisualization& EditorVisualization, const ::Chaos::FClothVisualizationNoGC& Visualization, FPrimitiveDrawInterface* PDI)
+		FClothVisualizationDebugDraw::CreateLambda([](const FClothEditorSimulationVisualization& EditorVisualization, const ::Chaos::FClothVisualizationNoGC& Visualization, const ::Chaos::FClothVisualizationNoGC::FDrawContext& Context)
 		{
 			const FString* const AccessoryMeshName = EditorVisualization.GetCurrentlySelectedAccessoryMesh();
-			Visualization.DrawAccessoryMeshNormals(PDI, AccessoryMeshName ? FName(*AccessoryMeshName) : NAME_None, EditorVisualization.GetAccessoryMeshNormalLength());
+			Visualization.DrawAccessoryMeshNormals(Context, AccessoryMeshName ? FName(*AccessoryMeshName) : NAME_None, EditorVisualization.GetAccessoryMeshNormalLength());
 		}),
 		LOCTEXT("ChaosVisName_DrawAccessoryMeshNormals", "Accessory Mesh Normals"), 
 		LOCTEXT("ChaosVisName_DrawAccessoryMeshNormals_ToolTip", "Draw an accessory mesh's normals."), 
@@ -502,9 +524,9 @@ const FVisualizationOption FVisualizationOption::OptionData[] =
 			EditorVisualization.ExtendViewportShowMenuAccessoryMeshNormalsLength(MenuBuilder);
 		})),
 	FVisualizationOption(
-		FClothVisualizationDebugDraw::CreateLambda([](const FClothEditorSimulationVisualization&, const ::Chaos::FClothVisualizationNoGC& Visualization, FPrimitiveDrawInterface* PDI)
+		FClothVisualizationDebugDraw::CreateLambda([](const FClothEditorSimulationVisualization&, const ::Chaos::FClothVisualizationNoGC& Visualization, const ::Chaos::FClothVisualizationNoGC::FDrawContext& Context)
 		{
-			Visualization.DrawInpaintWeightsMatched(PDI);
+			Visualization.DrawInpaintWeightsMatched(Context);
 		}),
 		LOCTEXT("ChaosVisName_DrawInpaintWeightsMatched", "Transfer Skin Weights Node: Matched Vertices"), 
 		LOCTEXT("ChaosVisName_DrawInpaintWeightsMatched_ToolTip", "When transferring weights using the InpaintWeights method, will highlight the vertices for which we copied the weights directly from the source mesh. For all other vertices, the weights were computed automatically."), 
@@ -512,8 +534,6 @@ const FVisualizationOption FVisualizationOption::OptionData[] =
 };
 const uint32 FVisualizationOption::Count = sizeof(OptionData) / sizeof(FVisualizationOption);
 
-
-
 } // namespace Private
 
 FClothEditorSimulationVisualization::FClothEditorSimulationVisualization()
@@ -620,7 +640,7 @@ void FClothEditorSimulationVisualization::ExtendViewportShowMenu(FMenuBuilder& M
 }
 
 
-void FClothEditorSimulationVisualization::ExtendViewportShowMenuNameSelector(FMenuBuilder& MenuBuilder, FNameSelectionData& SelectionData)
+void FClothEditorSimulationVisualization::ExtendViewportShowMenuNameSelector(FMenuBuilder& MenuBuilder, FNameSelectionData& SelectionData, const FText& ToolTipText)
 {
 	SelectionData.Selector =
 		SNew(STextComboBox)
@@ -628,25 +648,27 @@ void FClothEditorSimulationVisualization::ExtendViewportShowMenuNameSelector(FMe
 		.OnSelectionChanged(STextComboBox::FOnTextSelectionChanged::CreateLambda([&SelectionData](TSharedPtr<FString> Selection, ESelectInfo::Type SelectInfo)
 			{
 				SelectionData.CurrentlySelectedName = Selection;
-			}));
+			}))
+		.InitiallySelectedItem(SelectionData.CurrentlySelectedName);
 
 	MenuBuilder.AddMenuEntry(FUIAction(), SelectionData.Selector.ToSharedRef(),
-		NAME_None, LOCTEXT("WeightMapNameSelection", "Select weight map to draw."), EUserInterfaceActionType::None);
+		NAME_None, ToolTipText, EUserInterfaceActionType::None);
 }
 
-void FClothEditorSimulationVisualization::ExtendViewportShowMenuSpinBox(FMenuBuilder& MenuBuilder, float& Value, const float MinValue, const float MaxValue, const float MinSliderValue, const float MaxSliderValue)
+template<typename T>
+void FClothEditorSimulationVisualization::ExtendViewportShowMenuSpinBox(FMenuBuilder& MenuBuilder, T& Value, const T MinValue, const T MaxValue, const T MinSliderValue, const T MaxSliderValue, const FText& ToolTipText)
 {
 	TSharedRef<SWidget> SpinBox =
-		SNew(SSpinBox<float>)
+		SNew(SSpinBox<T>)
 		.Value(Value)
-		.OnValueChanged_Lambda([&Value](float NewValue) { Value = NewValue; })
+		.OnValueChanged_Lambda([&Value](T NewValue) { Value = NewValue; })
 		.MinValue(MinValue)
 		.MaxValue(MaxValue)
 		.MinSliderValue(MinSliderValue)
 		.MaxSliderValue(MaxSliderValue);
 
 	MenuBuilder.AddMenuEntry(FUIAction(), SpinBox,
-		NAME_None, LOCTEXT("SetNormalLength", "Set the normal display length."), EUserInterfaceActionType::None);
+		NAME_None, ToolTipText, EUserInterfaceActionType::None);
 }
 
 void FClothEditorSimulationVisualization::RefreshMenusForClothComponent(const UChaosClothComponent* ClothComponent)
@@ -654,10 +676,12 @@ void FClothEditorSimulationVisualization::RefreshMenusForClothComponent(const UC
 	TArray<FString> WeightMaps;
 	TArray<FString> MorphTargets;
 	TArray<FString> AccessoryMeshes;
+	TArray<FString> SelectionSets;
 	if (const ::Chaos::FClothVisualizationNoGC* const Visualization = ClothComponent && ClothComponent->GetClothSimulationProxy() ? ClothComponent->GetClothSimulationProxy()->GetClothVisualization() : nullptr)
 	{
 		WeightMaps = Visualization->GetAllWeightMapNames();
 		MorphTargets = Visualization->GetAllMorphTargetNames();
+		SelectionSets = Visualization->GetAllVertexSetNames();
 		const TArray<FName> AccessoryMeshNames = Visualization->GetAllAccessoryMeshNames();
 		AccessoryMeshes.Reserve(AccessoryMeshNames.Num());
 		for (const FName& AccessoryMeshName : AccessoryMeshNames)
@@ -700,6 +724,7 @@ void FClothEditorSimulationVisualization::RefreshMenusForClothComponent(const UC
 	RefreshMenuNames(WeightMaps, WeightMapSelection);
 	RefreshMenuNames(MorphTargets, MorphTargetSelection);
 	RefreshMenuNames(AccessoryMeshes, AccessoryMeshSelection);
+	RefreshMenuNames(SelectionSets, FilterSetSelection);
 }
 
 void FClothEditorSimulationVisualization::DebugDrawSimulation(const UChaosClothComponent* ClothComponent, FPrimitiveDrawInterface* PDI)
@@ -710,11 +735,24 @@ void FClothEditorSimulationVisualization::DebugDrawSimulation(const UChaosClothC
 		return;
 	}
 
+	::Chaos::FClothVisualizationNoGC::FDrawContext Context;
+	Context.PDI = PDI;
+
+	if (Flags[Private::FVisualizationOption::SingleVertexOptionIndex])
+	{
+		Context.Filter.SingleVertex = GetSingleVertexSelectionValue();
+	}
+	if (Flags[Private::FVisualizationOption::FilterSetOptionIndex])
+	{
+		const FString* VertexSet = GetCurrentlySelectedFilterSetSelection();
+		Context.Filter.VertexSet = VertexSet ? *VertexSet : FString();
+	}
+
 	for (int32 OptionIndex = 0; OptionIndex < Private::FVisualizationOption::Count; ++OptionIndex)
 	{
 		if (Flags[OptionIndex])
 		{
-			Private::FVisualizationOption::OptionData[OptionIndex].ClothVisualizationDebugDraw.ExecuteIfBound(*this, *Visualization, PDI);
+			Private::FVisualizationOption::OptionData[OptionIndex].ClothVisualizationDebugDraw.ExecuteIfBound(*this, *Visualization, Context);
 		}
 	}
 }
@@ -726,11 +764,25 @@ void FClothEditorSimulationVisualization::DebugDrawSimulationTexts(const UChaosC
 	{
 		return;
 	}
+	::Chaos::FClothVisualizationNoGC::FDrawTextsContext Context;
+	Context.Canvas = Canvas;
+	Context.SceneView = SceneView;
+
+	if (Flags[Private::FVisualizationOption::SingleVertexOptionIndex])
+	{
+		Context.Filter.SingleVertex = GetSingleVertexSelectionValue();
+	}
+	if (Flags[Private::FVisualizationOption::FilterSetOptionIndex])
+	{
+		const FString* VertexSet = GetCurrentlySelectedFilterSetSelection();
+		Context.Filter.VertexSet = VertexSet ? *VertexSet : FString();
+	}
+
 	for (int32 OptionIndex = 0; OptionIndex < Private::FVisualizationOption::Count; ++OptionIndex)
 	{
 		if (Flags[OptionIndex])
 		{
-			Private::FVisualizationOption::OptionData[OptionIndex].ClothVisualizationDebugDrawTexts.ExecuteIfBound(*this, *Visualization, Canvas, SceneView);
+			Private::FVisualizationOption::OptionData[OptionIndex].ClothVisualizationDebugDrawTexts.ExecuteIfBound(*this, *Visualization, Context);
 		}
 	}
 }
@@ -791,6 +843,67 @@ void FClothEditorSimulationVisualization::ShowClothSections(UChaosClothComponent
 		}
 	}
 }
+
+void FClothEditorSimulationVisualization::ExtendViewportShowMenuWeightMapSelector(FMenuBuilder& MenuBuilder)
+{
+	ExtendViewportShowMenuNameSelector(MenuBuilder, WeightMapSelection, LOCTEXT("WeightMapNameSelection", "Select weight map to draw."));
+}
+
+void FClothEditorSimulationVisualization::ExtendViewportShowMenuMorphTargetSelector(FMenuBuilder& MenuBuilder)
+{
+	ExtendViewportShowMenuNameSelector(MenuBuilder, MorphTargetSelection, LOCTEXT("MorphTargetNameSelection", "Select morph target to draw."));
+}
+
+void FClothEditorSimulationVisualization::ExtendViewportShowMenuAccessoryMeshSelector(FMenuBuilder& MenuBuilder)
+{
+	ExtendViewportShowMenuNameSelector(MenuBuilder, AccessoryMeshSelection, LOCTEXT("AccessoryMeshNameSelection", "Select accessory mesh to draw."));
+}
+
+void FClothEditorSimulationVisualization::ExtendViewportShowMenuFilterSetSelection(FMenuBuilder& MenuBuilder)
+{
+	ExtendViewportShowMenuNameSelector(MenuBuilder, FilterSetSelection, LOCTEXT("FilterSetSelectionName", "Select selection set to filter all debug draw with."));
+}
+
+void FClothEditorSimulationVisualization::ExtendViewportShowMenuPointNormalsLength(FMenuBuilder& MenuBuilder)
+{
+	ExtendViewportShowMenuSpinBox(MenuBuilder, PointNormalLength, 0.f, FLT_MAX, 0.f, 40.f, LOCTEXT("SetNormalLength", "Set the normal display length."));
+}
+
+void FClothEditorSimulationVisualization::ExtendViewportShowMenuAnimatedNormalsLength(FMenuBuilder& MenuBuilder)
+{
+	ExtendViewportShowMenuSpinBox(MenuBuilder, AnimatedNormalLength, 0.f, FLT_MAX, 0.f, 40.f, LOCTEXT("SetNormalLength", "Set the normal display length."));
+}
+
+void FClothEditorSimulationVisualization::ExtendViewportShowMenuAerodynamicsLengthScale(FMenuBuilder& MenuBuilder)
+{
+	ExtendViewportShowMenuSpinBox(MenuBuilder, AerodynamicsLengthScale, 0.f, FLT_MAX, 0.f, 40.f, LOCTEXT("SetLengthScale", "Scale the vector display length."));
+}
+
+void FClothEditorSimulationVisualization::ExtendViewportShowMenuAccessoryMeshNormalsLength(FMenuBuilder& MenuBuilder)
+{
+	ExtendViewportShowMenuSpinBox(MenuBuilder, AccessoryMeshNormalLength, 0.f, FLT_MAX, 0.f, 40.f, LOCTEXT("SetNormalLength", "Set the normal display length."));
+}
+
+void FClothEditorSimulationVisualization::ExtendViewportShowMenuWindVelocityLength(FMenuBuilder& MenuBuilder)
+{
+	ExtendViewportShowMenuSpinBox(MenuBuilder, WindVelocityLengthScale, 0.f, FLT_MAX, 0.f, 40.f, LOCTEXT("SetLengthScale", "Scale the vector display length."));
+}
+
+void FClothEditorSimulationVisualization::ExtendViewportShowMenuLocalSpaceLengthScale(FMenuBuilder& MenuBuilder)
+{
+	ExtendViewportShowMenuSpinBox(MenuBuilder, LocalSpaceLengthScale, 0.f, FLT_MAX, 0.f, 40.f, LOCTEXT("SetLengthScale", "Scale the vector display length."));
+}
+
+void FClothEditorSimulationVisualization::ExtendViewportShowMenuVelocityScaleLength(FMenuBuilder& MenuBuilder)
+{
+	ExtendViewportShowMenuSpinBox(MenuBuilder, VelocityScaleLength, 0.f, FLT_MAX, 0.f, 40.f, LOCTEXT("SetLengthScale", "Scale the vector display length."));
+}
+
+void FClothEditorSimulationVisualization::ExtendViewportShowMenuSingleVertexSelection(FMenuBuilder& MenuBuilder)
+{
+	ExtendViewportShowMenuSpinBox(MenuBuilder, SingleVertexSelection,0, INT_MAX, 0, 1000, LOCTEXT("SetSingleVertex", "Set single vertex debug filter."));
+}
+
 } // namespace UE::Chaos::ClothAsset
 
 #undef LOCTEXT_NAMESPACE
diff --git a/Engine/Plugins/ChaosClothAssetEditorCore/Source/ChaosClothAssetEditor/Private/ChaosClothAsset/ClothEditorSimulationVisualization.h b/Engine/Plugins/ChaosClothAssetEditorCore/Source/ChaosClothAssetEditor/Private/ChaosClothAsset/ClothEditorSimulationVisualization.h
--- a/Engine/Plugins/ChaosClothAssetEditorCore/Source/ChaosClothAssetEditor/Private/ChaosClothAsset/ClothEditorSimulationVisualization.h
+++ b/Engine/Plugins/ChaosClothAssetEditorCore/Source/ChaosClothAssetEditor/Private/ChaosClothAsset/ClothEditorSimulationVisualization.h
@@ -38,51 +38,25 @@ public:
 
 	// WeightMaps 
 	const FString* GetCurrentlySelectedWeightMap() const { return WeightMapSelection.CurrentlySelectedName.Get(); }
-	void ExtendViewportShowMenuWeightMapSelector(FMenuBuilder& MenuBuilder)
-	{
-		ExtendViewportShowMenuNameSelector(MenuBuilder, WeightMapSelection);
-	}
+	void ExtendViewportShowMenuWeightMapSelector(FMenuBuilder& MenuBuilder);
 	// Morph Targets
 	const FString* GetCurrentlySelectedMorphTarget() const { return MorphTargetSelection.CurrentlySelectedName.Get(); }
-	void ExtendViewportShowMenuMorphTargetSelector(FMenuBuilder& MenuBuilder)
-	{
-		ExtendViewportShowMenuNameSelector(MenuBuilder, MorphTargetSelection);
-	}
+	void ExtendViewportShowMenuMorphTargetSelector(FMenuBuilder& MenuBuilder);
 	// Accessory Meshes
 	const FString* GetCurrentlySelectedAccessoryMesh() const { return AccessoryMeshSelection.CurrentlySelectedName.Get(); }
-	void ExtendViewportShowMenuAccessoryMeshSelector(FMenuBuilder& MenuBuilder)
-	{
-		ExtendViewportShowMenuNameSelector(MenuBuilder, AccessoryMeshSelection);
-	}
+	void ExtendViewportShowMenuAccessoryMeshSelector(FMenuBuilder& MenuBuilder);
+	// DebugFilterSet
+	const FString* GetCurrentlySelectedFilterSetSelection() const { return FilterSetSelection.CurrentlySelectedName.Get(); }
+	void ExtendViewportShowMenuFilterSetSelection(FMenuBuilder& MenuBuilder);
 	// Normals
-	void ExtendViewportShowMenuPointNormalsLength(FMenuBuilder& MenuBuilder)
-	{
-		ExtendViewportShowMenuSpinBox(MenuBuilder, PointNormalLength, 0.f, FLT_MAX, 0.f, 40.f);
-	}
-	void ExtendViewportShowMenuAnimatedNormalsLength(FMenuBuilder& MenuBuilder)
-	{
-		ExtendViewportShowMenuSpinBox(MenuBuilder, AnimatedNormalLength, 0.f, FLT_MAX, 0.f, 40.f);
-	}
-	void ExtendViewportShowMenuAerodynamicsLengthScale(FMenuBuilder& MenuBuilder)
-	{
-		ExtendViewportShowMenuSpinBox(MenuBuilder, AerodynamicsLengthScale, 0.f, FLT_MAX, 0.f, 40.f);
-	}
-	void ExtendViewportShowMenuAccessoryMeshNormalsLength(FMenuBuilder& MenuBuilder)
-	{
-		ExtendViewportShowMenuSpinBox(MenuBuilder, AccessoryMeshNormalLength, 0.f, FLT_MAX, 0.f, 40.f);
-	}
-	void ExtendViewportShowMenuWindVelocityLength(FMenuBuilder& MenuBuilder)
-	{
-		ExtendViewportShowMenuSpinBox(MenuBuilder, WindVelocityLengthScale, 0.f, FLT_MAX, 0.f, 40.f);
-	}
-	void ExtendViewportShowMenuLocalSpaceLengthScale(FMenuBuilder& MenuBuilder)
-	{
-		ExtendViewportShowMenuSpinBox(MenuBuilder, LocalSpaceLengthScale , 0.f, FLT_MAX, 0.f, 40.f);
-	}
-	void ExtendViewportShowMenuVelocityScaleLength(FMenuBuilder& MenuBuilder)
-	{
-		ExtendViewportShowMenuSpinBox(MenuBuilder, VelocityScaleLength, 0.f, FLT_MAX, 0.f, 40.f);
-	}
+	void ExtendViewportShowMenuPointNormalsLength(FMenuBuilder& MenuBuilder);
+	void ExtendViewportShowMenuAnimatedNormalsLength(FMenuBuilder& MenuBuilder);
+	void ExtendViewportShowMenuAerodynamicsLengthScale(FMenuBuilder& MenuBuilder);
+	void ExtendViewportShowMenuAccessoryMeshNormalsLength(FMenuBuilder& MenuBuilder);
+	void ExtendViewportShowMenuWindVelocityLength(FMenuBuilder& MenuBuilder);
+	void ExtendViewportShowMenuLocalSpaceLengthScale(FMenuBuilder& MenuBuilder);
+	void ExtendViewportShowMenuVelocityScaleLength(FMenuBuilder& MenuBuilder);
+	void ExtendViewportShowMenuSingleVertexSelection(FMenuBuilder& MenuBuilder);
 	float GetPointNormalLength() const { return PointNormalLength; }
 	float GetAnimatedNormalLength() const { return AnimatedNormalLength; }
 	float GetAccessoryMeshNormalLength() const { return AccessoryMeshNormalLength; }
@@ -90,6 +64,7 @@ public:
 	float GetWindVelocityLengthScale() const { return WindVelocityLengthScale; }
 	float GetVelocityScaleLength() const { return VelocityScaleLength; }
 	float GetLocalSpaceLengthScale() const { return LocalSpaceLengthScale; }
+	int32 GetSingleVertexSelectionValue() const { return SingleVertexSelection; }
 
 private:
 	struct FNameSelectionData
@@ -106,8 +81,9 @@ private:
 	bool ShouldDisableSimulation() const;
 	/** Show/hide all cloth sections for the specified mesh compoment. */
 	void ShowClothSections(UChaosClothComponent* ClothComponent, bool bIsClothSectionsVisible) const;
-	void ExtendViewportShowMenuSpinBox(FMenuBuilder& MenuBuilder, float& Value, const float MinValue, const float MaxValue, const float MinSliderValue, const float MaxSliderValue);
-	void ExtendViewportShowMenuNameSelector(FMenuBuilder& MenuBuilder, FNameSelectionData& SelectionData);
+	template<typename T>
+	void ExtendViewportShowMenuSpinBox(FMenuBuilder& MenuBuilder, T& Value, const T MinValue, const T MaxValue, const T MinSliderValue, const T MaxSliderValue, const FText& ToolTipText);
+	void ExtendViewportShowMenuNameSelector(FMenuBuilder& MenuBuilder, FNameSelectionData& SelectionData, const FText& ToolTipText);
 
 private:
 	/** Flags used to store the checked status for the visualization options. */
@@ -115,6 +91,7 @@ private:
 	FNameSelectionData WeightMapSelection;
 	FNameSelectionData MorphTargetSelection;
 	FNameSelectionData AccessoryMeshSelection;
+	FNameSelectionData FilterSetSelection;
 	float PointNormalLength = 7.f;
 	float AnimatedNormalLength = 7.f;
 	float AccessoryMeshNormalLength = 7.f;
@@ -122,5 +99,6 @@ private:
 	float WindVelocityLengthScale = .1f;
 	float VelocityScaleLength = .1f;
 	float LocalSpaceLengthScale = 1.f;
+	int32 SingleVertexSelection = 0;
 };
 } // namespace UE::Chaos::ClothAsset
