Set default group name for Cloth selection set nodes (SimVertices3D). Correctly start Selection Tool in vertex selection mode when you have an empty selection of vertices.

This also fixes the ensure when you first create a Selection node.

If your collection only has a RenderCollection, the tool will just not start until you change the GroupName to a valid render group.

#Synced-CL 52021766

diff --git a/Engine/Plugins/ChaosClothAssetDataflowNodes/Source/ChaosClothAssetDataflowNodes/Public/ChaosClothAsset/SelectionNode.h b/Engine/Plugins/ChaosClothAssetDataflowNodes/Source/ChaosClothAssetDataflowNodes/Public/ChaosClothAsset/SelectionNode.h
--- a/Engine/Plugins/ChaosClothAssetDataflowNodes/Source/ChaosClothAssetDataflowNodes/Public/ChaosClothAsset/SelectionNode.h
+++ b/Engine/Plugins/ChaosClothAssetDataflowNodes/Source/ChaosClothAssetDataflowNodes/Public/ChaosClothAsset/SelectionNode.h
@@ -6,6 +6,7 @@
 #include "Dataflow/DataflowFunctionProperty.h"
 #include "GeometryCollection/ManagedArrayCollection.h"
 #include "ChaosClothAsset/AddWeightMapNode.h"
+#include "ChaosClothAsset/ClothCollectionGroup.h"
 #include "SelectionNode.generated.h"
 
 namespace UE_DEPRECATED(5.5, "Use UE::Dataflow instead.") Dataflow {}
@@ -43,7 +44,7 @@ struct FChaosClothAssetNodeSelectionGroup
 	GENERATED_USTRUCT_BODY()
 
 	UPROPERTY(EditAnywhere, Category = "Selection Group")
-	FString Name;
+	FString Name = TEXT("SimVertices3D");
 };
 
 /** How the sets stored on the SelectionNode should be applied to an existing set. If no set exists, it is treated as empty.*/
diff --git a/Engine/Plugins/ChaosClothAssetEditorCore/Source/ChaosClothAssetEditorTools/Private/ChaosClothAsset/ClothMeshSelectionTool.cpp b/Engine/Plugins/ChaosClothAssetEditorCore/Source/ChaosClothAssetEditorTools/Private/ChaosClothAsset/ClothMeshSelectionTool.cpp
--- a/Engine/Plugins/ChaosClothAssetEditorCore/Source/ChaosClothAssetEditorTools/Private/ChaosClothAsset/ClothMeshSelectionTool.cpp
+++ b/Engine/Plugins/ChaosClothAssetEditorCore/Source/ChaosClothAssetEditorTools/Private/ChaosClothAsset/ClothMeshSelectionTool.cpp
@@ -333,7 +333,10 @@ void UClothMeshSelectionTool::Setup()
 	constexpr bool bBroadcastChange = false;
 	SelectionMechanic->SetSelection(ExistingNodeSelection, bBroadcastChange);
 
-	if (ExistingNodeSelection.SelectedCornerIDs.Num() > 0)
+	if (ExistingNodeSelection.SelectedCornerIDs.Num() > 0 || 
+		SelectionNodeToUpdate->Group.Name == UE::Chaos::ClothAsset::ClothCollectionGroup::SimVertices3D || 
+		SelectionNodeToUpdate->Group.Name == UE::Chaos::ClothAsset::ClothCollectionGroup::SimVertices2D ||
+		SelectionNodeToUpdate->Group.Name == UE::Chaos::ClothAsset::ClothCollectionGroup::RenderVertices)
 	{
 		check(ExistingNodeSelection.SelectedEdgeIDs.Num() == 0);
 		check(ExistingNodeSelection.SelectedGroupIDs.Num() == 0);
